Fix Clang compile error due to missing symbols.

diff --git a/Engine/Source/Editor/AnimationEditMode/Private/AnimationEditMode.cpp b/Engine/Source/Editor/AnimationEditMode/Private/AnimationEditMode.cpp
--- a/Engine/Source/Editor/AnimationEditMode/Private/AnimationEditMode.cpp
+++ b/Engine/Source/Editor/AnimationEditMode/Private/AnimationEditMode.cpp
@@ -24,6 +24,11 @@ void UAnimationEditModeContext::GetOnScreenDebugInfo(TArray<FText>& OutDebugInfo
 	return EditMode->GetOnScreenDebugInfo(OutDebugInfo);
 };
 
+FAnimationEditMode::FAnimationEditMode()
+	: AnimationEditModeContext(UAnimationEditModeContext::CreateFor(this))
+{
+}
+
 void FAnimationEditMode::Enter()
 {
 	FEdMode::Enter();
diff --git a/Engine/Source/Editor/AnimationEditMode/Public/AnimationEditMode.h b/Engine/Source/Editor/AnimationEditMode/Public/AnimationEditMode.h
--- a/Engine/Source/Editor/AnimationEditMode/Public/AnimationEditMode.h
+++ b/Engine/Source/Editor/AnimationEditMode/Public/AnimationEditMode.h
@@ -36,9 +36,14 @@ private:
 class ANIMATIONEDITMODE_API FAnimationEditMode : public FEdMode, public IAnimationEditContext
 {
 public:
+	FAnimationEditMode();
+
+	FAnimationEditMode(const FAnimationEditMode&) = delete;
+	FAnimationEditMode& operator=(const FAnimationEditMode&) = delete;
+
 	virtual void Enter() override;
 	virtual void Exit() override;
 	virtual void AddReferencedObjects(FReferenceCollector& Collector) override;
 private:
-	TObjectPtr<UAnimationEditModeContext> AnimationEditModeContext{UAnimationEditModeContext::CreateFor(this)};
+	TObjectPtr<UAnimationEditModeContext> AnimationEditModeContext;
 };
